Fixing some reordering/etc issues reported by clang.

diff --git a/Engine/Plugins/Runtime/LiveEditorListenServer/Source/LiveEditorListenServer/Private/LiveEditorListenServerEngineHooks.h b/Engine/Plugins/Runtime/LiveEditorListenServer/Source/LiveEditorListenServer/Private/LiveEditorListenServerEngineHooks.h
--- a/Engine/Plugins/Runtime/LiveEditorListenServer/Source/LiveEditorListenServer/Private/LiveEditorListenServerEngineHooks.h
+++ b/Engine/Plugins/Runtime/LiveEditorListenServer/Source/LiveEditorListenServer/Private/LiveEditorListenServerEngineHooks.h
@@ -11,7 +11,7 @@ class FCreateListener : public FUObjectArray::FUObjectCreateListener
 {
 public:
 	FCreateListener( FLiveEditorListenServer *_Owner );
-	~FCreateListener();
+	virtual ~FCreateListener();
 	virtual void NotifyUObjectCreated(const class UObjectBase *Object, int32 Index);
 
 private:
@@ -22,7 +22,7 @@ class FDeleteListener : public FUObjectArray::FUObjectDeleteListener
 {
 public:
 	FDeleteListener( FLiveEditorListenServer *_Owner );
-	~FDeleteListener();
+	virtual ~FDeleteListener();
 	virtual void NotifyUObjectDeleted(const class UObjectBase *Object, int32 Index);
 
 private:
diff --git a/Engine/Source/Editor/UnrealEd/Public/UnrealWidget.h b/Engine/Source/Editor/UnrealEd/Public/UnrealWidget.h
--- a/Engine/Source/Editor/UnrealEd/Public/UnrealWidget.h
+++ b/Engine/Source/Editor/UnrealEd/Public/UnrealWidget.h
@@ -193,8 +193,8 @@ private:
 	struct FThickArcParams
 	{
 		FThickArcParams(FPrimitiveDrawInterface* InPDI, const FVector& InPosition, UMaterialInterface* InMaterial, const float InInnerRadius, const float InOuterRadius)
-			: PDI(InPDI)
-			, Position(InPosition)
+			: Position(InPosition)
+			, PDI(InPDI)
 			, Material(InMaterial)
 			, InnerRadius(InInnerRadius)
 			, OuterRadius(InOuterRadius)
diff --git a/Engine/Source/Runtime/Core/Public/Linux/LinuxPlatform.h b/Engine/Source/Runtime/Core/Public/Linux/LinuxPlatform.h
--- a/Engine/Source/Runtime/Core/Public/Linux/LinuxPlatform.h
+++ b/Engine/Source/Runtime/Core/Public/Linux/LinuxPlatform.h
@@ -68,6 +68,10 @@ typedef FLinuxPlatformTypes FPlatformTypes;
 #define FINAL	 final
 #define ABSTRACT abstract
 
+// DLL export and import definitions
+#define DLLEXPORT			__attribute__((visibility("default")))
+#define DLLIMPORT			__attribute__((visibility("default")))
+
 // Alignment.
 #define GCC_PACK(n)			__attribute__((packed,aligned(n)))
 #define GCC_ALIGN(n)		__attribute__((aligned(n)))
