Fix LiveEditorListenServer so that objects created after the initial LiveEditor change has come through will also pick up LiveEditor changes

diff --git a/Engine/Plugins/Runtime/LiveEditorListenServer/Source/LiveEditorListenServer/Private/LiveEditorListenServer.cpp b/Engine/Plugins/Runtime/LiveEditorListenServer/Source/LiveEditorListenServer/Private/LiveEditorListenServer.cpp
--- a/Engine/Plugins/Runtime/LiveEditorListenServer/Source/LiveEditorListenServer/Private/LiveEditorListenServer.cpp
+++ b/Engine/Plugins/Runtime/LiveEditorListenServer/Source/LiveEditorListenServer/Private/LiveEditorListenServer.cpp
@@ -6,50 +6,6 @@
 
 namespace nLiveEditorListenServer
 {
-	/*
-	static UProperty *GetPropertyByNameRecurse( UStruct *InStruct, const FString &TokenString, void ** hContainerPtr )
-	{
-		FString FirstToken;
-		FString RemainingTokens;
-		int32 SplitIndex;
-		if ( TokenString.FindChar( '.', SplitIndex ) )
-		{
-			FirstToken = TokenString.LeftChop( TokenString.Len()-SplitIndex );
-			RemainingTokens = TokenString.RightChop(SplitIndex+1);
-		}
-		else
-		{
-			FirstToken = TokenString;
-			RemainingTokens = FString(TEXT(""));
-		}
-
-		for (TFieldIterator<UProperty> PropertyIt(InStruct, EFieldIteratorFlags::IncludeSuper); PropertyIt; ++PropertyIt)
-		{
-			UProperty* Property = *PropertyIt;
-			FName PropertyName = Property->GetFName();
-			if ( FirstToken == PropertyName.ToString() )
-			{
-				if ( RemainingTokens.Len() == 0 )
-				{
-					check( *hContainerPtr != NULL );
-					return Property;
-				}
-				else
-				{
-					UStructProperty *StructProp = Cast<UStructProperty>(Property);
-					if ( StructProp )
-					{
-						check( *hContainerPtr != NULL );
-						*hContainerPtr = Property->ContainerPtrToValuePtr<void>(*hContainerPtr);
-						return GetPropertyByNameRecurse( StructProp->Struct, RemainingTokens, hContainerPtr );
-					}
-				}
-			}
-		}
-
-		return NULL;
-	}
-	*/
 	static UProperty *GetPropertyByNameRecurse( UStruct *InStruct, const FString &TokenString, void ** hContainerPtr, int32 &OutArrayIndex )
 	{
 		FString FirstToken;
@@ -296,6 +252,8 @@ void FLiveEditorListenServer::Tick( float DeltaTime )
 			}
 		}
 	}
+
+	ObjectCache.EvaluatePendingCreations();
 #endif
 }
 
diff --git a/Engine/Plugins/Runtime/LiveEditorListenServer/Source/LiveEditorListenServer/Private/LiveEditorListenServerObjectCache.cpp b/Engine/Plugins/Runtime/LiveEditorListenServer/Source/LiveEditorListenServer/Private/LiveEditorListenServerObjectCache.cpp
--- a/Engine/Plugins/Runtime/LiveEditorListenServer/Source/LiveEditorListenServer/Private/LiveEditorListenServerObjectCache.cpp
+++ b/Engine/Plugins/Runtime/LiveEditorListenServer/Source/LiveEditorListenServer/Private/LiveEditorListenServerObjectCache.cpp
@@ -142,7 +142,7 @@ void FObjectCache::EvaluatePendingCreations()
 			UObject *AsObject = const_cast<UObject*>( static_cast<const UObject*>(ObjectBase) );
 	
 			UWorld *World = GEngine->GetWorldFromContextObject( AsObject, false );
-			if ( World == NULL || World->WorldType != EWorldType::PIE )
+			if ( World == NULL || World->WorldType != EWorldType::Game )
 				continue;
 
 			for ( TMultiMap< const UObject*, TWeakObjectPtr<UObject> >::TConstIterator ObjIt(ObjectLookupCache); ObjIt; ++ObjIt )
