On Remote clients, cache a transaction history of latest updates to LiveEdited values.
When new objects are created, if they are objects that are of the type being tuned, stuff them with latest values

diff --git a/Engine/Plugins/Runtime/LiveEditorListenServer/Source/LiveEditorListenServer/Private/LiveEditorListenServer.cpp b/Engine/Plugins/Runtime/LiveEditorListenServer/Source/LiveEditorListenServer/Private/LiveEditorListenServer.cpp
--- a/Engine/Plugins/Runtime/LiveEditorListenServer/Source/LiveEditorListenServer/Private/LiveEditorListenServer.cpp
+++ b/Engine/Plugins/Runtime/LiveEditorListenServer/Source/LiveEditorListenServer/Private/LiveEditorListenServer.cpp
@@ -3,6 +3,7 @@
 #include "LiveEditorListenServerPrivatePCH.h"
 #include "Networking.h"
 #include "DefaultValueHelper.h"
+#include "LiveEditorTransactionHistory.h"
 
 namespace nLiveEditorListenServer
 {
@@ -97,8 +98,87 @@ namespace nLiveEditorListenServer
 
 		return Prop;
 	}
+
+	static void SetPropertyValue( UObject *Target, const FString &PropertyName, const FString &PropertyValue )
+	{
+
+		if ( Target == NULL )
+			return;
+
+		void *ContainerPtr = Target;
+		int32 ArrayIndex;
+		UProperty *Prop = nLiveEditorListenServer::GetPropertyByName( Target, Target->GetClass(), PropertyName, &ContainerPtr, ArrayIndex );
+		if ( !Prop || !Prop->IsA( UNumericProperty::StaticClass() ) )
+		{
+			return;
+		}
+
+		check( ContainerPtr != NULL );
+
+		if ( Prop->IsA( UByteProperty::StaticClass() ) )
+		{
+			UByteProperty *NumericProp = CastChecked<UByteProperty>(Prop);
+			uint8 Value = FCString::Atoi( *PropertyValue );
+			NumericProp->SetPropertyValue_InContainer(ContainerPtr, Value, ArrayIndex);
+		}
+		else if ( Prop->IsA( UInt8Property::StaticClass() ) )
+		{
+			UInt8Property *NumericProp = CastChecked<UInt8Property>(Prop);
+			int32 Value = FCString::Atoi( *PropertyValue );
+			NumericProp->SetPropertyValue_InContainer(ContainerPtr, Value, ArrayIndex);
+		}
+		else if ( Prop->IsA( UInt16Property::StaticClass() ) )
+		{
+			UInt16Property *NumericProp = CastChecked<UInt16Property>(Prop);
+			int16 Value = FCString::Atoi( *PropertyValue );
+			NumericProp->SetPropertyValue_InContainer(ContainerPtr, Value, ArrayIndex);
+		}
+		else if ( Prop->IsA( UIntProperty::StaticClass() ) )
+		{
+			UIntProperty *NumericProp = CastChecked<UIntProperty>(Prop);
+			int32 Value = FCString::Atoi( *PropertyValue );
+			NumericProp->SetPropertyValue_InContainer(ContainerPtr, Value, ArrayIndex);
+		}
+		else if ( Prop->IsA( UInt64Property::StaticClass() ) )
+		{
+			UInt64Property *NumericProp = CastChecked<UInt64Property>(Prop);
+			int64 Value = FCString::Atoi64( *PropertyValue );
+			NumericProp->SetPropertyValue_InContainer(ContainerPtr, Value, ArrayIndex);
+		}
+		else if ( Prop->IsA( UUInt16Property::StaticClass() ) )
+		{
+			UUInt16Property *NumericProp = CastChecked<UUInt16Property>(Prop);
+			uint16 Value = FCString::Atoi( *PropertyValue );
+			NumericProp->SetPropertyValue_InContainer(ContainerPtr, Value, ArrayIndex);
+		}
+		else if ( Prop->IsA( UUInt32Property::StaticClass() ) )
+		{
+			UUInt32Property *NumericProp = CastChecked<UUInt32Property>(Prop);
+			uint32 Value = FCString::Atoi( *PropertyValue );
+			NumericProp->SetPropertyValue_InContainer(ContainerPtr, Value, ArrayIndex);
+		}
+		else if ( Prop->IsA( UInt64Property::StaticClass() ) )
+		{
+			UInt64Property *NumericProp = CastChecked<UInt64Property>(Prop);
+			uint64 Value = FCString::Atoi64( *PropertyValue );
+			NumericProp->SetPropertyValue_InContainer(ContainerPtr, Value, ArrayIndex);
+		}
+		else if ( Prop->IsA( UFloatProperty::StaticClass() ) )
+		{
+			UFloatProperty *NumericProp = CastChecked<UFloatProperty>(Prop);
+			float Value = FCString::Atof( *PropertyValue );
+			NumericProp->SetPropertyValue_InContainer(ContainerPtr, Value, ArrayIndex);
+		}
+		else if ( Prop->IsA( UDoubleProperty::StaticClass() ) )
+		{
+			UDoubleProperty *NumericProp = CastChecked<UDoubleProperty>(Prop);
+			double Value = FCString::Atod( *PropertyValue );
+			NumericProp->SetPropertyValue_InContainer(ContainerPtr, Value, ArrayIndex);
+		}
+	}
 }
 
+
 DEFINE_LOG_CATEGORY_STATIC(LiveEditorListenServer, Log, All);
 #define DEFAULT_LISTEN_ENDPOINT FIPv4Endpoint(FIPv4Address(127, 0, 0, 1), LIVEEDITORLISTENSERVER_DEFAULT_PORT)
 
@@ -108,7 +188,8 @@ FLiveEditorListenServer::FLiveEditorListenServer()
 :	ObjectCreateListener(NULL),
 	ObjectDeleteListener(NULL),
 	TickObject(NULL),
-	Listener(NULL)
+	Listener(NULL),
+	TransactionHistory(NULL)
 {
 }
 
@@ -121,6 +202,8 @@ void FLiveEditorListenServer::StartupModule()
 		return;
 	}
 
+	TransactionHistory = new FLiveEditorTransactionHistory();
+
 	Listener = new FTcpListener( DEFAULT_LISTEN_ENDPOINT );
 	Listener->OnConnectionAccepted().BindRaw(this, &FLiveEditorListenServer::HandleListenerConnectionAccepted);
 
@@ -159,6 +242,12 @@ void FLiveEditorListenServer::RemoveHooks()
 		Listener = NULL;
 	}
 
+	if ( TransactionHistory )
+	{
+		delete TransactionHistory;
+		TransactionHistory = NULL;
+	}
+
 	if ( !PendingClients.IsEmpty() )
 	{
 		FSocket *Client = NULL;
@@ -249,89 +338,31 @@ void FLiveEditorListenServer::Tick( float DeltaTime )
 				*Datagram << Message;
 
 				ReplicateChanges( Message );
+				TransactionHistory->AppendTransaction( Message.Payload.ClassName, Message.Payload.PropertyName, Message.Payload.PropertyValue );
 			}
 		}
 	}
 
-	ObjectCache.EvaluatePendingCreations();
-#endif
-}
-
-static void SetPropertyValue( UObject *Target, const FString &PropertyName, const FString &PropertyValue )
-{
-	if ( Target == NULL )
-		return;
+	//Check for any newly Tracked objects
+	TArray<UObject*> NewTrackedObjects;
+	ObjectCache.EvaluatePendingCreations( NewTrackedObjects );
 
-	void *ContainerPtr = Target;
-	int32 ArrayIndex;
-	UProperty *Prop = nLiveEditorListenServer::GetPropertyByName( Target, Target->GetClass(), PropertyName, &ContainerPtr, ArrayIndex );
-	if ( !Prop || !Prop->IsA( UNumericProperty::StaticClass() ) )
+	//If there are any newly tracked objects, stuff them with the list of latest changes we've accumulated so far
+	for ( TArray<UObject*>::TIterator NewObjectIt(NewTrackedObjects); NewObjectIt; ++NewObjectIt )
 	{
-		return;
-	}
-
-	check( ContainerPtr != NULL );
+		UObject *NewObject = *NewObjectIt;
+		TMap<FString,FString> Transactions;
+		TransactionHistory->FindTransactionsForObject( NewObject, Transactions );
 
-	if ( Prop->IsA( UByteProperty::StaticClass() ) )
-	{
-		UByteProperty *NumericProp = CastChecked<UByteProperty>(Prop);
-		uint8 Value = FCString::Atoi( *PropertyValue );
-		NumericProp->SetPropertyValue_InContainer(ContainerPtr, Value, ArrayIndex);
-	}
-	else if ( Prop->IsA( UInt8Property::StaticClass() ) )
-	{
-		UInt8Property *NumericProp = CastChecked<UInt8Property>(Prop);
-		int32 Value = FCString::Atoi( *PropertyValue );
-		NumericProp->SetPropertyValue_InContainer(ContainerPtr, Value, ArrayIndex);
-	}
-	else if ( Prop->IsA( UInt16Property::StaticClass() ) )
-	{
-		UInt16Property *NumericProp = CastChecked<UInt16Property>(Prop);
-		int16 Value = FCString::Atoi( *PropertyValue );
-		NumericProp->SetPropertyValue_InContainer(ContainerPtr, Value, ArrayIndex);
-	}
-	else if ( Prop->IsA( UIntProperty::StaticClass() ) )
-	{
-		UIntProperty *NumericProp = CastChecked<UIntProperty>(Prop);
-		int32 Value = FCString::Atoi( *PropertyValue );
-		NumericProp->SetPropertyValue_InContainer(ContainerPtr, Value, ArrayIndex);
-	}
-	else if ( Prop->IsA( UInt64Property::StaticClass() ) )
-	{
-		UInt64Property *NumericProp = CastChecked<UInt64Property>(Prop);
-		int64 Value = FCString::Atoi64( *PropertyValue );
-		NumericProp->SetPropertyValue_InContainer(ContainerPtr, Value, ArrayIndex);
-	}
-	else if ( Prop->IsA( UUInt16Property::StaticClass() ) )
-	{
-		UUInt16Property *NumericProp = CastChecked<UUInt16Property>(Prop);
-		uint16 Value = FCString::Atoi( *PropertyValue );
-		NumericProp->SetPropertyValue_InContainer(ContainerPtr, Value, ArrayIndex);
-	}
-	else if ( Prop->IsA( UUInt32Property::StaticClass() ) )
-	{
-		UUInt32Property *NumericProp = CastChecked<UUInt32Property>(Prop);
-		uint32 Value = FCString::Atoi( *PropertyValue );
-		NumericProp->SetPropertyValue_InContainer(ContainerPtr, Value, ArrayIndex);
-	}
-	else if ( Prop->IsA( UInt64Property::StaticClass() ) )
-	{
-		UInt64Property *NumericProp = CastChecked<UInt64Property>(Prop);
-		uint64 Value = FCString::Atoi64( *PropertyValue );
-		NumericProp->SetPropertyValue_InContainer(ContainerPtr, Value, ArrayIndex);
-	}
-	else if ( Prop->IsA( UFloatProperty::StaticClass() ) )
-	{
-		UFloatProperty *NumericProp = CastChecked<UFloatProperty>(Prop);
-		float Value = FCString::Atof( *PropertyValue );
-		NumericProp->SetPropertyValue_InContainer(ContainerPtr, Value, ArrayIndex);
-	}
-	else if ( Prop->IsA( UDoubleProperty::StaticClass() ) )
-	{
-		UDoubleProperty *NumericProp = CastChecked<UDoubleProperty>(Prop);
-		double Value = FCString::Atod( *PropertyValue );
-		NumericProp->SetPropertyValue_InContainer(ContainerPtr, Value, ArrayIndex);
+		for ( TMap<FString,FString>::TConstIterator TransactionIt(Transactions); TransactionIt; ++TransactionIt )
+		{
+			const FString &PropertyName = (*TransactionIt).Key;
+			const FString &PropertyValue = (*TransactionIt).Value;
+			nLiveEditorListenServer::SetPropertyValue( NewObject, PropertyName, PropertyValue );
+		}
 	}
+	
+#endif
 }
 
 void FLiveEditorListenServer::ReplicateChanges( const nLiveEditorListenServer::FNetworkMessage &Message )
@@ -342,7 +373,7 @@ void FLiveEditorListenServer::ReplicateChanges( const nLiveEditorListenServer::F
 		UClass *BaseClass = FindObject<UClass>(ANY_PACKAGE, *Message.Payload.ClassName);
 		if ( BaseClass )
 		{
-			SetPropertyValue( BaseClass->ClassDefaultObject, Message.Payload.PropertyName, Message.Payload.PropertyValue );
+			nLiveEditorListenServer::SetPropertyValue( BaseClass->ClassDefaultObject, Message.Payload.PropertyName, Message.Payload.PropertyValue );
 
 			TArray< TWeakObjectPtr<UObject> > Replicants;
 			ObjectCache.FindObjectDependants( BaseClass->ClassDefaultObject, Replicants );
@@ -353,7 +384,7 @@ void FLiveEditorListenServer::ReplicateChanges( const nLiveEditorListenServer::F
 
 				UObject *Object = (*It).Get();
 				check( Object->IsA(BaseClass) );
-				SetPropertyValue( Object, Message.Payload.PropertyName, Message.Payload.PropertyValue );
+				nLiveEditorListenServer::SetPropertyValue( Object, Message.Payload.PropertyName, Message.Payload.PropertyValue );
 			}
 		}
 	}
diff --git a/Engine/Plugins/Runtime/LiveEditorListenServer/Source/LiveEditorListenServer/Private/LiveEditorListenServer.h b/Engine/Plugins/Runtime/LiveEditorListenServer/Source/LiveEditorListenServer/Private/LiveEditorListenServer.h
--- a/Engine/Plugins/Runtime/LiveEditorListenServer/Source/LiveEditorListenServer/Private/LiveEditorListenServer.h
+++ b/Engine/Plugins/Runtime/LiveEditorListenServer/Source/LiveEditorListenServer/Private/LiveEditorListenServer.h
@@ -4,6 +4,8 @@
 
 #include "LiveEditorListenServerObjectCache.h"
 
+class FLiveEditorTransactionHistory;
+
 class FLiveEditorListenServer : public ILiveEditorListenServer
 {
 public:
@@ -38,4 +40,6 @@ private:
 	class FTcpListener *Listener;
 	TQueue<class FSocket*, EQueueMode::Mpsc> PendingClients;
 	TArray<class FSocket*> Clients;
+
+	FLiveEditorTransactionHistory *TransactionHistory;
 };
diff --git a/Engine/Plugins/Runtime/LiveEditorListenServer/Source/LiveEditorListenServer/Private/LiveEditorListenServerObjectCache.cpp b/Engine/Plugins/Runtime/LiveEditorListenServer/Source/LiveEditorListenServer/Private/LiveEditorListenServerObjectCache.cpp
--- a/Engine/Plugins/Runtime/LiveEditorListenServer/Source/LiveEditorListenServer/Private/LiveEditorListenServerObjectCache.cpp
+++ b/Engine/Plugins/Runtime/LiveEditorListenServer/Source/LiveEditorListenServer/Private/LiveEditorListenServerObjectCache.cpp
@@ -127,7 +127,7 @@ void FObjectCache::OnObjectDeletion( const class UObjectBase* ObjectBase )
 	}
 }
 
-void FObjectCache::EvaluatePendingCreations()
+void FObjectCache::EvaluatePendingCreations( TArray<UObject*> &NewTrackedObjects )
 {
 	if ( !CacheStarted() || ObjectLookupCache.Num() <= 0 || GEngine == NULL )
 		return;
@@ -153,6 +153,7 @@ void FObjectCache::EvaluatePendingCreations()
 
 				ObjectLookupCache.AddUnique( Key, AsObject );
 				TrackedObjects.Add( AsObject );
+				NewTrackedObjects.Add( AsObject );
 			}
 		}
 	}
diff --git a/Engine/Plugins/Runtime/LiveEditorListenServer/Source/LiveEditorListenServer/Private/LiveEditorListenServerObjectCache.h b/Engine/Plugins/Runtime/LiveEditorListenServer/Source/LiveEditorListenServer/Private/LiveEditorListenServerObjectCache.h
--- a/Engine/Plugins/Runtime/LiveEditorListenServer/Source/LiveEditorListenServer/Private/LiveEditorListenServerObjectCache.h
+++ b/Engine/Plugins/Runtime/LiveEditorListenServer/Source/LiveEditorListenServer/Private/LiveEditorListenServerObjectCache.h
@@ -15,7 +15,8 @@ namespace nLiveEditorListenServer
 		void EndCache();
 		void OnObjectCreation( const class UObjectBase* ObjectBase );
 		void OnObjectDeletion( const class UObjectBase* ObjectBase );
-		void EvaluatePendingCreations();
+
+		void EvaluatePendingCreations( TArray<UObject*> &NewTrackedObjects );
 
 	private:
 		bool CacheStarted() const;
diff --git a/Engine/Plugins/Runtime/LiveEditorListenServer/Source/LiveEditorListenServer/Private/LiveEditorTransactionHistory.cpp b/Engine/Plugins/Runtime/LiveEditorListenServer/Source/LiveEditorListenServer/Private/LiveEditorTransactionHistory.cpp
new file mode 100644
--- /dev/null
+++ b/Engine/Plugins/Runtime/LiveEditorListenServer/Source/LiveEditorListenServer/Private/LiveEditorTransactionHistory.cpp
@@ -0,0 +1,45 @@
+// Copyright 1998-2014 Epic Games, Inc. All Rights Reserved.
+
+#include "LiveEditorListenServerPrivatePCH.h"
+#include "LiveEditorTransactionHistory.h"
+
+FLiveEditorTransactionHistory::FLiveEditorTransactionHistory()
+{
+}
+
+void FLiveEditorTransactionHistory::AppendTransaction( const FString &ClassName, const FString &PropertyName, const FString &PropertyValue )
+{
+	TMap<FString,FString> *ExistingValues = EditHistory.Find( ClassName );
+	if ( ExistingValues != NULL )
+	{
+		ExistingValues->Add( PropertyName, PropertyValue );
+	}
+	else
+	{
+		TMap<FString,FString> NewMap;
+		NewMap.Add( PropertyName, PropertyValue );
+		EditHistory.Add( ClassName, NewMap );
+	}
+}
+
+void FLiveEditorTransactionHistory::FindTransactionsForObject( const UObject *Object, TMap<FString,FString> &OutValues ) const
+{
+	for ( TMap< FString, TMap<FString,FString> >::TConstIterator ClassIt(EditHistory); ClassIt; ++ClassIt )
+	{
+		const FString &ClassName = (*ClassIt).Key;
+		UClass *BaseClass = FindObject<UClass>(ANY_PACKAGE, *ClassName);
+		if ( BaseClass == NULL || !Object->IsA(BaseClass) )
+			continue;
+
+		const TMap<FString,FString> *ExistingValues = EditHistory.Find( ClassName );
+		if ( ExistingValues == NULL )
+			continue;
+
+		for ( TMap<FString,FString>::TConstIterator TransactionIt(*ExistingValues); TransactionIt; ++TransactionIt )
+		{
+			const FString &PropertyName = (*TransactionIt).Key;
+			const FString &PropertyValue = (*TransactionIt).Value;
+			OutValues.Add( PropertyName, PropertyValue );
+		}
+	}
+}
diff --git a/Engine/Plugins/Runtime/LiveEditorListenServer/Source/LiveEditorListenServer/Private/LiveEditorTransactionHistory.h b/Engine/Plugins/Runtime/LiveEditorListenServer/Source/LiveEditorListenServer/Private/LiveEditorTransactionHistory.h
new file mode 100644
--- /dev/null
+++ b/Engine/Plugins/Runtime/LiveEditorListenServer/Source/LiveEditorListenServer/Private/LiveEditorTransactionHistory.h
@@ -0,0 +1,16 @@
+
+// Copyright 1998-2014 Epic Games, Inc. All Rights Reserved.
+
+#pragma once
+
+class FLiveEditorTransactionHistory
+{
+public:
+	FLiveEditorTransactionHistory();
+
+	void AppendTransaction( const FString &ClassName, const FString &PropertyName, const FString &PropertyValue );
+	void FindTransactionsForObject( const UObject *Object, TMap<FString,FString> &OutValues ) const;
+
+private:
+	TMap< FString, TMap<FString,FString> > EditHistory;
+};
